[Animator Kit]-Initial add of Animator Kit plugin to house animation helper rigs and deformer rigs.

#virtualized

diff --git a/Engine/Plugins/Animation/AnimatorKit/AnimatorKit.uplugin b/Engine/Plugins/Animation/AnimatorKit/AnimatorKit.uplugin
new file mode 100644
--- /dev/null
+++ b/Engine/Plugins/Animation/AnimatorKit/AnimatorKit.uplugin
@@ -0,0 +1,37 @@
+{
+	"FileVersion": 3,
+	"Version": 1,
+	"VersionName": "1.0",
+	"FriendlyName": "Animator Kit",
+	"Description": "Utilities for Animating in Unreal with Sequencer",
+	"Category": "Animation",
+	"CreatedBy": "Epic Games, Inc.",
+	"CreatedByURL": "https://epicgames.com",
+	"DocsURL": "",
+	"MarketplaceURL": "",
+	"SupportURL": "",
+	"CanContainContent": true,
+	"IsBetaVersion": true,
+	"IsExperimentalVersion": false,
+	"Installed": false,
+	"EnabledByDefault": false,
+	"Plugins": [
+		{
+			"Name": "ControlRig",
+			"Enabled": true
+		},
+		{
+			"Name": "ControlRigSpline",
+			"Enabled": true
+		},
+		{
+			"Name": "RigVM",
+			"Enabled": true
+		},
+		{
+			"Name": "DeformerGraph",
+			"Enabled": true
+		}
+	],
+	"NoCode": true
+}
\ No newline at end of file
