[Animator Kit Plugin] Added CR variant tags for deformer rigs and utility rigs

diff --git a/Engine/Plugins/Animation/AnimatorKit/Config/Editor.ini b/Engine/Plugins/Animation/AnimatorKit/Config/Editor.ini
new file mode 100644
--- /dev/null
+++ b/Engine/Plugins/Animation/AnimatorKit/Config/Editor.ini
@@ -0,0 +1,3 @@
+[/Script/RigVM.RigVMProjectSettings]
++VariantTags=(Name="AnimatorKit_Deformer",Label="Deformer",ToolTip=NSLOCTEXT("RigVMSettings", "AnimatorKit_Deformer_ToolTip", "Deformer rigs can be used deform characters in Sequencer. These rigs have built-in Deformer Graphs"), Color=(R=0.033105,G=0.266356,B=0.791298,A=1.000000),bShowInUserInterface=True,bMarksSubjectAsInvalid=False)
++VariantTags=(Name="AnimatorKit_Utility",Label="Utility",ToolTip=NSLOCTEXT("RigVMSettings", "AnimatorKit_Utility_ToolTip", "Utility rigs can be used for on-demand rigging in Sequencer"), Color=(R=0.73791,G=0.099899,B=0.318547,A=1.000000),bShowInUserInterface=True,bMarksSubjectAsInvalid=False)
