[Animator Kit] Enable GizmoFramework plugin when enabling animator kit

diff --git a/Engine/Plugins/Animation/AnimatorKit/AnimatorKit.uplugin b/Engine/Plugins/Animation/AnimatorKit/AnimatorKit.uplugin
--- a/Engine/Plugins/Animation/AnimatorKit/AnimatorKit.uplugin
+++ b/Engine/Plugins/Animation/AnimatorKit/AnimatorKit.uplugin
@@ -31,6 +31,10 @@
 		{
 			"Name": "DeformerGraph",
 			"Enabled": true
+		},
+		{
+			"Name": "GizmoFramework",
+			"Enabled": true 
 		}
 	],
 	"NoCode": true
