Animation mode & AnimatorKit: expose focus settings in the editor settings
- todo: add more settings to drive the other animation related parameters

diff --git a/Engine/Plugins/Animation/AnimatorKit/AnimatorKit.uplugin b/Engine/Plugins/Animation/AnimatorKit/AnimatorKit.uplugin
--- a/Engine/Plugins/Animation/AnimatorKit/AnimatorKit.uplugin
+++ b/Engine/Plugins/Animation/AnimatorKit/AnimatorKit.uplugin
@@ -37,5 +37,11 @@
 			"Enabled": true 
 		}
 	],
-	"NoCode": true
+	"Modules" : [
+		{
+			"Name" : "AnimatorKitSettings",
+			"Type" : "Editor",
+			"LoadingPhase" : "PostEngineInit"
+		}
+	]
 }
\ No newline at end of file
diff --git a/Engine/Plugins/Animation/AnimatorKit/Source/AnimatorKitSettings/AnimatorKitSettings.Build.cs b/Engine/Plugins/Animation/AnimatorKit/Source/AnimatorKitSettings/AnimatorKitSettings.Build.cs
new file mode 100644
--- /dev/null
+++ b/Engine/Plugins/Animation/AnimatorKit/Source/AnimatorKitSettings/AnimatorKitSettings.Build.cs
@@ -0,0 +1,21 @@
+// Copyright Epic Games, Inc. All Rights Reserved.
+
+namespace UnrealBuildTool.Rules
+{
+	public class AnimatorKitSettings : ModuleRules
+	{
+		public AnimatorKitSettings(ReadOnlyTargetRules Target) : base(Target)
+		{
+			PrivateDependencyModuleNames.AddRange(
+				new string[]
+				{
+					"Core",
+					"CoreUObject",
+					"DeveloperSettings",
+					"Engine",
+					"ControlRigEditor"
+				}
+			);
+		}
+	}
+}
diff --git a/Engine/Plugins/Animation/AnimatorKit/Source/AnimatorKitSettings/Private/AnimatorKitSettings.cpp b/Engine/Plugins/Animation/AnimatorKit/Source/AnimatorKitSettings/Private/AnimatorKitSettings.cpp
new file mode 100644
--- /dev/null
+++ b/Engine/Plugins/Animation/AnimatorKit/Source/AnimatorKitSettings/Private/AnimatorKitSettings.cpp
@@ -0,0 +1,31 @@
+// Copyright Epic Games, Inc. All Rights Reserved.
+
+#include "AnimatorKitSettings.h"
+
+UAnimatorKitSettings::UAnimatorKitSettings(const FObjectInitializer& ObjectInitlaizer)
+	: UDeveloperSettings(ObjectInitlaizer)
+{}
+
+UAnimatorKitSettings::FOnUpdateSettings UAnimatorKitSettings::OnSettingsChange;
+
+void UAnimatorKitSettings::PostInitProperties()
+{
+	Super::PostInitProperties(); 
+
+	if (IsTemplate())
+	{
+		ImportConsoleVariableValues();
+	}
+}
+
+void UAnimatorKitSettings::PostEditChangeProperty(FPropertyChangedEvent& PropertyChangedEvent)
+{
+	Super::PostEditChangeProperty(PropertyChangedEvent);
+	
+	if (PropertyChangedEvent.Property)
+	{
+		ExportValuesToConsoleVariables(PropertyChangedEvent.Property);
+	}
+	
+	OnSettingsChange.Broadcast(this);
+}
diff --git a/Engine/Plugins/Animation/AnimatorKit/Source/AnimatorKitSettings/Private/AnimatorKitSettingsModule.cpp b/Engine/Plugins/Animation/AnimatorKit/Source/AnimatorKitSettings/Private/AnimatorKitSettingsModule.cpp
new file mode 100644
--- /dev/null
+++ b/Engine/Plugins/Animation/AnimatorKit/Source/AnimatorKitSettings/Private/AnimatorKitSettingsModule.cpp
@@ -0,0 +1,25 @@
+// Copyright Epic Games, Inc. All Rights Reserved.
+
+#include "AnimatorKitSettingsModule.h"
+
+#include "AnimatorKitSettings.h"
+#include "Modules/ModuleManager.h"
+
+void FAnimatorKitSettingsModule::StartupModule()
+{
+	UAnimatorKitSettings::OnSettingsChange.AddRaw(this, &FAnimatorKitSettingsModule::UpdateSettings);
+	UpdateSettings(GetDefault<UAnimatorKitSettings>());
+}
+
+void FAnimatorKitSettingsModule::ShutdownModule()
+{
+	UAnimatorKitSettings::OnSettingsChange.RemoveAll(this);
+}
+
+void FAnimatorKitSettingsModule::UpdateSettings(const UAnimatorKitSettings* InSettings)
+{
+	// tbd
+}
+
+
+IMPLEMENT_MODULE(FAnimatorKitSettingsModule, AnimatorKitSettings)
diff --git a/Engine/Plugins/Animation/AnimatorKit/Source/AnimatorKitSettings/Public/AnimatorKitSettings.h b/Engine/Plugins/Animation/AnimatorKit/Source/AnimatorKitSettings/Public/AnimatorKitSettings.h
new file mode 100644
--- /dev/null
+++ b/Engine/Plugins/Animation/AnimatorKit/Source/AnimatorKitSettings/Public/AnimatorKitSettings.h
@@ -0,0 +1,23 @@
+#pragma once
+
+#include "CoreMinimal.h"
+#include "Engine/DeveloperSettings.h"
+
+#include "AnimatorKitSettings.generated.h"
+
+UCLASS(config=EditorPerProjectUserSettings, meta = (Experimental, DisplayName = "Animator Kit Settings"))
+class ANIMATORKITSETTINGS_API UAnimatorKitSettings : public UDeveloperSettings
+{
+	GENERATED_UCLASS_BODY()
+
+public:
+
+	UPROPERTY(config, EditAnywhere, Category = "Animation Settings|Focus", meta=(ConsoleVariable="AnimMode.PendingFocusMode"))
+	bool bEnableFocusMode = false;
+    
+	DECLARE_MULTICAST_DELEGATE_OneParam(FOnUpdateSettings,  const UAnimatorKitSettings*);
+	static FOnUpdateSettings OnSettingsChange;
+
+	virtual void PostInitProperties() override;
+	virtual void PostEditChangeProperty(FPropertyChangedEvent& PropertyChangedEvent) override;
+};
\ No newline at end of file
diff --git a/Engine/Plugins/Animation/AnimatorKit/Source/AnimatorKitSettings/Public/AnimatorKitSettingsModule.h b/Engine/Plugins/Animation/AnimatorKit/Source/AnimatorKitSettings/Public/AnimatorKitSettingsModule.h
new file mode 100644
--- /dev/null
+++ b/Engine/Plugins/Animation/AnimatorKit/Source/AnimatorKitSettings/Public/AnimatorKitSettingsModule.h
@@ -0,0 +1,16 @@
+// Copyright Epic Games, Inc. All Rights Reserved.
+
+#pragma once
+
+#include "Modules/ModuleInterface.h"
+
+class UAnimatorKitSettings;
+
+class FAnimatorKitSettingsModule : public IModuleInterface
+{
+	/** IModuleInterface implementation */
+	virtual void StartupModule() override;
+	virtual void ShutdownModule() override;
+
+	void UpdateSettings(const UAnimatorKitSettings* InSettings);
+};
diff --git a/Engine/Plugins/Animation/ControlRig/Source/ControlRigEditor/Private/EditMode/ControlRigEditMode.cpp b/Engine/Plugins/Animation/ControlRig/Source/ControlRigEditor/Private/EditMode/ControlRigEditMode.cpp
--- a/Engine/Plugins/Animation/ControlRig/Source/ControlRigEditor/Private/EditMode/ControlRigEditMode.cpp
+++ b/Engine/Plugins/Animation/ControlRig/Source/ControlRigEditor/Private/EditMode/ControlRigEditMode.cpp
@@ -3871,6 +3871,7 @@ void FControlRigEditMode::RegisterPendingFocusMode()
 			PendingFocusHandle = UseFocusMode->OnChangedDelegate().AddLambda(OnFocusModeChanged);
 		}		
 		OnFocusModeChanged(UseFocusMode);
+		UseFocusMode->ClearFlags(ECVF_SetByDeviceProfile);
 	}
 }
 
