fixed missing copyright

#rnx

diff --git a/Engine/Plugins/Animation/AnimatorKit/Source/AnimatorKitSettings/Public/AnimatorKitSettings.h b/Engine/Plugins/Animation/AnimatorKit/Source/AnimatorKitSettings/Public/AnimatorKitSettings.h
--- a/Engine/Plugins/Animation/AnimatorKit/Source/AnimatorKitSettings/Public/AnimatorKitSettings.h
+++ b/Engine/Plugins/Animation/AnimatorKit/Source/AnimatorKitSettings/Public/AnimatorKitSettings.h
@@ -1,3 +1,5 @@
+// Copyright Epic Games, Inc. All Rights Reserved.
+
 #pragma once
 
 #include "CoreMinimal.h"
