Animation mode & AnimatorKit: set focus mode to true by default

#rnx

diff --git a/Engine/Plugins/Animation/AnimatorKit/Source/AnimatorKitSettings/Public/AnimatorKitSettings.h b/Engine/Plugins/Animation/AnimatorKit/Source/AnimatorKitSettings/Public/AnimatorKitSettings.h
--- a/Engine/Plugins/Animation/AnimatorKit/Source/AnimatorKitSettings/Public/AnimatorKitSettings.h
+++ b/Engine/Plugins/Animation/AnimatorKit/Source/AnimatorKitSettings/Public/AnimatorKitSettings.h
@@ -15,7 +15,7 @@ class ANIMATORKITSETTINGS_API UAnimatorKitSettings : public UDeveloperSettings
 public:
 
 	UPROPERTY(config, EditAnywhere, Category = "Animation Settings|Focus", meta=(ConsoleVariable="AnimMode.PendingFocusMode"))
-	bool bEnableFocusMode = false;
+	bool bEnableFocusMode = true;
     
 	DECLARE_MULTICAST_DELEGATE_OneParam(FOnUpdateSettings,  const UAnimatorKitSettings*);
 	static FOnUpdateSettings OnSettingsChange;
