Animation mode & AnimatorKit: fixed focus mode value propagation
- only import the focus cvar value if it has been set with a higher priority than ECVF_SetByProjectSetting (e.g. ECVF_SetByDeviceProfile)

#rnx

diff --git a/Engine/Plugins/Animation/AnimatorKit/Source/AnimatorKitSettings/Private/AnimatorKitSettings.cpp b/Engine/Plugins/Animation/AnimatorKit/Source/AnimatorKitSettings/Private/AnimatorKitSettings.cpp
--- a/Engine/Plugins/Animation/AnimatorKit/Source/AnimatorKitSettings/Private/AnimatorKitSettings.cpp
+++ b/Engine/Plugins/Animation/AnimatorKit/Source/AnimatorKitSettings/Private/AnimatorKitSettings.cpp
@@ -2,8 +2,10 @@
 
 #include "AnimatorKitSettings.h"
 
-UAnimatorKitSettings::UAnimatorKitSettings(const FObjectInitializer& ObjectInitlaizer)
-	: UDeveloperSettings(ObjectInitlaizer)
+#include "HAL/IConsoleManager.h"
+
+UAnimatorKitSettings::UAnimatorKitSettings(const FObjectInitializer& ObjectInitializer)
+	: UDeveloperSettings(ObjectInitializer)
 {}
 
 UAnimatorKitSettings::FOnUpdateSettings UAnimatorKitSettings::OnSettingsChange;
@@ -14,7 +16,23 @@ void UAnimatorKitSettings::PostInitProperties()
 
 	if (IsTemplate())
 	{
-		ImportConsoleVariableValues();
+		// only import the focus cvar value if it has been set with a higher priority than ECVF_SetByProjectSetting (e.g. ECVF_SetByDeviceProfile) 
+		// ImportConsoleVariableValues();
+		
+		const TCHAR* FocusModeName = TEXT("AnimMode.PendingFocusMode");
+		FProperty* FocusModeProp = GetClass()->FindPropertyByName(GET_MEMBER_NAME_CHECKED(UAnimatorKitSettings, bEnableFocusMode));
+		IConsoleVariable* FocusModeCVar = IConsoleManager::Get().FindConsoleVariable(FocusModeName);
+		if (FocusModeProp && FocusModeCVar)
+		{
+			if ((FocusModeCVar->GetFlags() & ECVF_SetByMask) < ECVF_SetByProjectSetting)
+			{
+				FocusModeCVar->Set(static_cast<int32>(bEnableFocusMode), ECVF_SetByProjectSetting);
+			}
+			else
+			{
+				FocusModeProp->ImportText_InContainer(*FocusModeCVar->GetString(), this, this, PPF_ConsoleVariable);
+			}
+		}
 	}
 }
 
