[AnimatorKit] - Update plugin settings and icon

- Set EnabledByDefault to true
- Updated plugin icon to AnimatorKitLogo128

#testing
#virtualized
#Synced-CL 52720788

diff --git a/Engine/Plugins/Animation/AnimatorKit/AnimatorKit.uplugin b/Engine/Plugins/Animation/AnimatorKit/AnimatorKit.uplugin
--- a/Engine/Plugins/Animation/AnimatorKit/AnimatorKit.uplugin
+++ b/Engine/Plugins/Animation/AnimatorKit/AnimatorKit.uplugin
@@ -10,11 +10,12 @@
 	"DocsURL": "",
 	"MarketplaceURL": "",
 	"SupportURL": "",
+	"PluginIconPath": "Resources/Icon128.png",
 	"CanContainContent": true,
 	"IsBetaVersion": true,
 	"IsExperimentalVersion": false,
 	"Installed": false,
-	"EnabledByDefault": false,
+	"EnabledByDefault": true,
 	"Plugins": [
 		{
 			"Name": "ControlRig",
@@ -24,6 +25,10 @@
 			"Name": "ControlRigSpline",
 			"Enabled": true
 		},
+		{
+			"Name": "ControlRigPhysics",
+			"Enabled": true
+		},
 		{
 			"Name": "RigVM",
 			"Enabled": true
@@ -34,7 +39,7 @@
 		},
 		{
 			"Name": "GizmoFramework",
-			"Enabled": true 
+			"Enabled": true
 		}
 	],
 	"Modules" : [
